Created a graphics debugging tool that can be used in the client to visualize draw call and primitive data. The tool can be opened with the command DrawPrimitiveDebugger.Open

diff --git a/Engine/Source/Developer/DrawPrimitiveDebugger/DrawPrimitiveDebugger.Build.cs b/Engine/Source/Developer/DrawPrimitiveDebugger/DrawPrimitiveDebugger.Build.cs
new file mode 100644
--- /dev/null
+++ b/Engine/Source/Developer/DrawPrimitiveDebugger/DrawPrimitiveDebugger.Build.cs
@@ -0,0 +1,39 @@
+﻿// Copyright Epic Games, Inc. All Rights Reserved.
+
+namespace UnrealBuildTool.Rules
+{
+	public class DrawPrimitiveDebugger : ModuleRules
+	{
+		public DrawPrimitiveDebugger(ReadOnlyTargetRules Target) : base(Target)
+		{
+			PublicDependencyModuleNames.AddRange(
+				new string[] {
+					"Core",
+					"CoreUObject",
+					"Engine"
+				});
+
+			PrivateDependencyModuleNames.AddRange(
+				new string[] {
+					"RenderCore",
+					"Renderer",
+					"InputCore",
+					"SlateCore",
+					"Slate",
+					"DeveloperSettings",
+				});
+
+			PrivateIncludePaths.AddRange(
+				new string[] {
+					"Developer/DrawPrimitiveDebugger/Private",
+					"Developer/Settings/Public",
+				});
+
+			// TODO: Decide whether Test builds should also remain restricted
+			if (Target.bBuildDeveloperTools || (Target.Configuration != UnrealTargetConfiguration.Shipping))
+			{
+				PrecompileForTargets = PrecompileTargetsType.Any;
+			}
+		}
+	}
+}
\ No newline at end of file
diff --git a/Engine/Source/Developer/DrawPrimitiveDebugger/Private/DrawPrimitiveDebuggerConfig.cpp b/Engine/Source/Developer/DrawPrimitiveDebugger/Private/DrawPrimitiveDebuggerConfig.cpp
new file mode 100644
--- /dev/null
+++ b/Engine/Source/Developer/DrawPrimitiveDebugger/Private/DrawPrimitiveDebuggerConfig.cpp
@@ -0,0 +1,13 @@
+﻿// Copyright Epic Games, Inc. All Rights Reserved.
+
+#include "DrawPrimitiveDebuggerConfig.h"
+
+void UDrawPrimitiveDebuggerUserSettings::SetFontSize(const int32 InFontSize)
+{
+	UDrawPrimitiveDebuggerUserSettings* Settings = GetMutableDefault<UDrawPrimitiveDebuggerUserSettings>();
+	Settings->FontSize = FMath::Clamp(InFontSize, 1, 100);
+
+	const FString ConfigFileName = Settings->GetProjectUserConfigFilename();
+	GConfig->SetInt(TEXT("/Script/DrawPrimitiveDebugger.DrawPrimitiveDebuggerUserSettings"), TEXT("FontSize"), Settings->FontSize, *ConfigFileName);
+	GConfig->Flush(false);
+}
diff --git a/Engine/Source/Developer/DrawPrimitiveDebugger/Private/DrawPrimitiveDebuggerModule.cpp b/Engine/Source/Developer/DrawPrimitiveDebugger/Private/DrawPrimitiveDebuggerModule.cpp
new file mode 100644
--- /dev/null
+++ b/Engine/Source/Developer/DrawPrimitiveDebugger/Private/DrawPrimitiveDebuggerModule.cpp
@@ -0,0 +1,177 @@
+﻿// Copyright Epic Games, Inc. All Rights Reserved.
+
+#include "CoreMinimal.h"
+#include "Modules/ModuleManager.h"
+#include "DrawPrimitiveDebugger.h"
+#include "SDrawPrimitiveDebugger.h"
+#include "ViewDebug.h"
+#include "Widgets/Docking/SDockTab.h"
+
+DEFINE_LOG_CATEGORY(LogDrawPrimitiveDebugger)
+
+#if !UE_BUILD_SHIPPING
+static FAutoConsoleCommand SummonDebuggerCmd(
+		TEXT("DrawPrimitiveDebugger.Open"),
+		TEXT("Summons the graphics debugger window."),
+		FConsoleCommandDelegate::CreateLambda([]() { IDrawPrimitiveDebugger::Get().OpenDebugWindow(); })
+		);
+/*static FAutoConsoleCommand EnableLiveCaptureCmd(
+		TEXT("DrawPrimitiveDebugger.EnableLiveCapture"),
+		TEXT("Enables live graphics data capture each frame."),
+		FConsoleCommandDelegate::CreateLambda([]() { IDrawPrimitiveDebugger::Get().EnableLiveCapture(); })
+		);
+static FAutoConsoleCommand DisableLiveCaptureCmd(
+		TEXT("DrawPrimitiveDebugger.DisableLiveCapture"),
+		TEXT("Disables live graphics data capture."),
+		FConsoleCommandDelegate::CreateLambda([]() { IDrawPrimitiveDebugger::Get().DisableLiveCapture(); })
+		);*/ // TODO: Re-enable these commands once live capture performance has been fixed
+static FAutoConsoleCommand TakeSnapshotCmd(
+		TEXT("DrawPrimitiveDebugger.Snapshot"),
+		TEXT("Updates the current view information for a single frame."),
+		FConsoleCommandDelegate::CreateLambda([]() { IDrawPrimitiveDebugger::Get().CaptureSingleFrame(); })
+		);
+#endif
+
+class FDrawPrimitiveDebuggerModule : public IDrawPrimitiveDebugger
+{
+public:
+	virtual void StartupModule() override;
+	virtual void ShutdownModule() override;
+
+	virtual void CaptureSingleFrame() override;
+	virtual bool IsLiveCaptureEnabled() const override;
+	virtual void EnableLiveCapture() override;
+	virtual void DisableLiveCapture() override;
+	virtual void OpenDebugWindow() override;
+	virtual void CloseDebugWindow() override;
+#if !UE_BUILD_SHIPPING
+	static const FViewDebugInfo& GetViewDebugInfo();
+#endif
+
+private:
+	bool bLiveCaptureEnabled = false;
+	FDelegateHandle UpdateDelegateHandle;
+	TSharedPtr<SDrawPrimitiveDebugger> DebuggerWidget;
+	TSharedPtr<SDockTab> DebuggerTab;
+
+	TSharedRef<SDockTab> MakeDrawPrimitiveDebuggerTab(const FSpawnTabArgs&);
+
+	void OnTabClosed(TSharedRef<SDockTab> Tab);
+	
+	void OnUpdateViewInformation();
+};
+
+IMPLEMENT_MODULE(FDrawPrimitiveDebuggerModule, DrawPrimitiveDebugger)
+
+void FDrawPrimitiveDebuggerModule::StartupModule()
+{
+	bLiveCaptureEnabled = false;
+#if !UE_BUILD_SHIPPING
+	UpdateDelegateHandle = FViewDebugInfo::Instance.AddUpdateHandler(this, &FDrawPrimitiveDebuggerModule::OnUpdateViewInformation);
+	FGlobalTabmanager::Get()->RegisterNomadTabSpawner("DrawPrimitiveDebugger", FOnSpawnTab::CreateRaw(this, &FDrawPrimitiveDebuggerModule::MakeDrawPrimitiveDebuggerTab) );
+#endif
+}
+
+void FDrawPrimitiveDebuggerModule::ShutdownModule()
+{
+#if !UE_BUILD_SHIPPING
+	FViewDebugInfo::Instance.RemoveUpdateHandler(UpdateDelegateHandle);
+#endif
+}
+
+void FDrawPrimitiveDebuggerModule::CaptureSingleFrame()
+{
+#if !UE_BUILD_SHIPPING
+	UE_LOG(LogDrawPrimitiveDebugger, Log, TEXT("Collecting a single frame graphics data capture"));
+	FViewDebugInfo::Instance.CaptureNextFrame();
+#endif
+}
+
+bool FDrawPrimitiveDebuggerModule::IsLiveCaptureEnabled() const
+{
+	return bLiveCaptureEnabled;
+}
+
+void FDrawPrimitiveDebuggerModule::EnableLiveCapture()
+{
+#if !UE_BUILD_SHIPPING
+	if (!bLiveCaptureEnabled)
+	{
+		UE_LOG(LogDrawPrimitiveDebugger, Log, TEXT("Enabling live graphics data capture"));
+		bLiveCaptureEnabled = true;
+		FViewDebugInfo::Instance.EnableLiveCapture();
+	}
+#endif
+}
+
+void FDrawPrimitiveDebuggerModule::DisableLiveCapture()
+{
+#if !UE_BUILD_SHIPPING
+	if (bLiveCaptureEnabled)
+	{
+		UE_LOG(LogDrawPrimitiveDebugger, Log, TEXT("Disabling live graphics data capture"));
+		bLiveCaptureEnabled = false;
+		FViewDebugInfo::Instance.DisableLiveCapture();
+	}
+#endif
+}
+
+#if !UE_BUILD_SHIPPING
+const FViewDebugInfo& FDrawPrimitiveDebuggerModule::GetViewDebugInfo()
+{
+	return FViewDebugInfo::Get();
+}
+#endif
+
+void FDrawPrimitiveDebuggerModule::OpenDebugWindow()
+{
+#if !UE_BUILD_SHIPPING
+	if (!DebuggerTab.IsValid())
+	{
+		UE_LOG(LogDrawPrimitiveDebugger, Log, TEXT("Opening the Draw Primitive Debugger"));
+		if (!FViewDebugInfo::Instance.HasEverUpdated()) CaptureSingleFrame();
+		FGlobalTabmanager::Get()->TryInvokeTab(FTabId("DrawPrimitiveDebugger"));
+	}
+#endif
+}
+
+void FDrawPrimitiveDebuggerModule::CloseDebugWindow()
+{
+#if !UE_BUILD_SHIPPING
+	if (DebuggerTab.IsValid())
+	{
+		UE_LOG(LogDrawPrimitiveDebugger, Log, TEXT("Closing the Draw Primitive Debugger"));
+		DebuggerTab->RequestCloseTab();
+	}
+#endif
+}
+
+TSharedRef<SDockTab> FDrawPrimitiveDebuggerModule::MakeDrawPrimitiveDebuggerTab(const FSpawnTabArgs&)
+{
+	DebuggerTab = SNew(SDockTab)
+		.TabRole(ETabRole::NomadTab)
+		.OnTabClosed(SDockTab::FOnTabClosedCallback::CreateRaw(this, &FDrawPrimitiveDebuggerModule::OnTabClosed));
+	
+	if (!DebuggerWidget.IsValid())
+	{
+		DebuggerWidget = SNew(SDrawPrimitiveDebugger);
+		DebuggerTab->SetContent(DebuggerWidget.ToSharedRef());
+	}
+	return DebuggerTab.ToSharedRef();
+}
+
+void FDrawPrimitiveDebuggerModule::OnTabClosed(TSharedRef<SDockTab> Tab)
+{
+	DebuggerTab.Reset();
+	DebuggerTab = nullptr;
+	DebuggerWidget.Reset();
+	DebuggerWidget = nullptr;
+}
+
+void FDrawPrimitiveDebuggerModule::OnUpdateViewInformation()
+{
+	if (DebuggerWidget && DebuggerWidget.IsValid())
+	{
+		DebuggerWidget->Refresh();
+	}
+}
\ No newline at end of file
diff --git a/Engine/Source/Developer/DrawPrimitiveDebugger/Private/SDrawPrimitiveDebugger.cpp b/Engine/Source/Developer/DrawPrimitiveDebugger/Private/SDrawPrimitiveDebugger.cpp
new file mode 100644
--- /dev/null
+++ b/Engine/Source/Developer/DrawPrimitiveDebugger/Private/SDrawPrimitiveDebugger.cpp
@@ -0,0 +1,503 @@
+﻿// Copyright Epic Games, Inc. All Rights Reserved.
+
+#include "SDrawPrimitiveDebugger.h"
+
+#include "DrawPrimitiveDebugger.h"
+#include "DrawPrimitiveDebuggerConfig.h"
+#include "SlateOptMacros.h"
+#include "Fonts/FontMeasure.h"
+#include "Widgets/Images/SImage.h"
+#include "Widgets/Layout/SScrollBox.h"
+
+#if !UE_BUILD_SHIPPING
+
+BEGIN_SLATE_FUNCTION_BUILD_OPTIMIZATION
+
+void SDrawPrimitiveDebugger::Construct(const FArguments& InArgs)
+{
+	const TSharedRef<SScrollBar> VerticalScrollBar = SNew(SScrollBar)
+		.Orientation(Orient_Vertical)
+		.Thickness(FVector2D(12.0f, 12.0f));
+
+	ColumnHeader = SNew(SHeaderRow).ResizeMode(ESplitterResizeMode::Fill);
+	const FName VisibilityColumn("Visible");
+	const FName PinColumn("Pin");
+	const FName NameColumn("Name");
+	const FName NumDrawsColumn("NumDraws");
+	const FName ActorClassColumn("ActorClass");
+	const FName ActorColumn("Actor");
+	const FName LocationColumn("Location");
+	const FName MaterialsColumn("Materials");
+	const FName LODColumn("LOD");
+	const FName TrianglesColumn("Triangles");
+	AddColumn(FText::FromString("Visible"), VisibilityColumn);
+	AddColumn(FText::FromString("Pinned"), PinColumn);
+	AddColumn(FText::FromString("Name"), NameColumn);
+	AddColumn(FText::FromString("Draws"), NumDrawsColumn);
+	AddColumn(FText::FromString("ActorClass"), ActorClassColumn);
+	AddColumn(FText::FromString("Actor"), ActorColumn);
+	AddColumn(FText::FromString("Location"), LocationColumn);
+	AddColumn(FText::FromString("Materials"), MaterialsColumn);
+	AddColumn(FText::FromString("LOD"), LODColumn);
+	AddColumn(FText::FromString("Triangles"), TrianglesColumn);
+
+	Refresh();
+
+	Table = SNew(SListView<FPrimitiveRowDataPtr>)
+		.ListItemsSource(&VisibleEntries)
+		.HeaderRow(ColumnHeader)
+		.OnGenerateRow(this, &SDrawPrimitiveDebugger::MakeRowWidget)
+		.ExternalScrollbar(VerticalScrollBar)
+		.Orientation(Orient_Vertical)
+		.ConsumeMouseWheel(EConsumeMouseWheel::Never)
+		.SelectionMode(ESelectionMode::Multi);
+	
+	ChildSlot
+	[
+		SNew(SVerticalBox)
+		+ SVerticalBox::Slot()
+		.AutoHeight()
+		[
+			SNew(SHorizontalBox)
+			+ SHorizontalBox::Slot()
+			.Padding(6, 6)
+			.HAlign(HAlign_Fill)
+			.FillWidth(2)
+			[
+				SAssignNew(SearchBox, SSearchBox)
+				.InitialText(this, &SDrawPrimitiveDebugger::GetFilterText)
+				.OnTextChanged(this, &SDrawPrimitiveDebugger::OnFilterTextChanged)
+				.OnTextCommitted(this, &SDrawPrimitiveDebugger::OnFilterTextCommitted)
+			]
+			+ SHorizontalBox::Slot()
+			.Padding(6, 6)
+			.AutoWidth()
+			[
+				SNew(SButton)
+				.Text(FText::FromString("Refresh"))
+				.IsEnabled(this, &SDrawPrimitiveDebugger::CanCaptureSingleFrame)
+				.OnClicked(this, &SDrawPrimitiveDebugger::OnRefreshClick)
+			]
+			+ SHorizontalBox::Slot()
+			.Padding(6, 6)
+			.AutoWidth()
+			[
+				SNew(SButton)
+				.Text(FText::FromString("Save to CSV"))
+				.OnClicked(this, &SDrawPrimitiveDebugger::OnSaveClick)
+			]
+			/*+ SHorizontalBox::Slot()
+			.Padding(6, 6)
+			.AutoWidth()
+			[
+				SNew(SCheckBox)
+				[
+					SNew(STextBlock)
+					.Text(FText::FromString("Enable Live Capture"))
+					.Font(FSlateFontInfo(FCoreStyle::GetDefaultFont(), UDrawPrimitiveDebuggerUserSettings::GetFontSize()))
+				]
+				.IsChecked(this, &SDrawPrimitiveDebugger::IsLiveCaptureChecked)
+				.OnCheckStateChanged(this, &SDrawPrimitiveDebugger::OnToggleLiveCapture)
+			]*/ // TODO: Re-enable after the performance issues have been fixed
+		]
+		+SVerticalBox::Slot()
+		.Padding(6, 6)
+		[
+			SNew(SScrollBox)
+			.Orientation(Orient_Vertical)
+			.ConsumeMouseWheel(EConsumeMouseWheel::Always)
+			+SScrollBox::Slot()
+			[
+				Table.ToSharedRef()
+			]
+		]
+	];
+}
+
+FText SDrawPrimitiveDebugger::GetFilterText() const
+{
+	return FilterText;
+}
+
+void SDrawPrimitiveDebugger::OnFilterTextChanged(const FText& InFilterText)
+{
+	FilterText = InFilterText;
+	UpdateVisibleRows();
+	if (Table.IsValid())
+	{
+		Table->RequestListRefresh();
+	}
+}
+
+void SDrawPrimitiveDebugger::OnFilterTextCommitted(const FText& NewText, ETextCommit::Type CommitInfo)
+{
+	if (CommitInfo == ETextCommit::OnCleared)
+	{
+		SearchBox->SetText(FText::GetEmpty());
+		OnFilterTextChanged(FText::GetEmpty());
+	}
+}
+
+TSharedRef<ITableRow> SDrawPrimitiveDebugger::MakeRowWidget(FPrimitiveRowDataPtr InRowDataPtr, const TSharedRef<STableViewBase>& OwnerTable)
+{
+	return SNew(SDrawPrimitiveDebuggerListViewRow, OwnerTable)
+		.DrawPrimitiveDebugger(SharedThis(this))
+		.RowDataPtr(InRowDataPtr);
+}
+
+void SDrawPrimitiveDebugger::UpdateVisibleRows()
+{
+	if (FilterText.IsEmptyOrWhitespace())
+	{
+		VisibleEntries = AvailableEntries;
+	}
+	else
+	{
+		VisibleEntries.Empty();
+
+		const FString& ActiveFilterString = FilterText.ToString();
+		for (const FPrimitiveRowDataPtr& RowData : AvailableEntries)
+		{
+			if (!RowData.IsValid()) continue;
+			bool bPassesFilter = false;
+			
+			if (RowData->Name.Contains(ActiveFilterString))
+			{
+				bPassesFilter = true;
+			}
+			else if (IsValid(RowData->Owner) &&
+				(RowData->Owner->GetClass()->GetName().Contains(ActiveFilterString) ||
+				RowData->Owner->GetFullName().Contains(ActiveFilterString)))
+			{
+				bPassesFilter = true;
+			}
+			else
+			{
+				for (const UMaterialInterface* Material : RowData->Materials)
+				{
+					if (IsValid(Material) && Material->GetName().Contains(ActiveFilterString))
+					{
+						bPassesFilter = true;
+						break;
+					}
+				}
+			}
+
+			if (bPassesFilter)
+			{
+				VisibleEntries.Add(RowData);
+			}
+		}
+	}
+	SortRows();
+}
+
+void SDrawPrimitiveDebugger::SortRows()
+{
+	VisibleEntries.Sort([this](FPrimitiveRowDataPtr A, FPrimitiveRowDataPtr B)
+	{
+		const bool bPinnedA = IsEntryPinned(A);
+		const bool bPinnedB = IsEntryPinned(B);
+		return (bPinnedA > bPinnedB) || ((bPinnedA == bPinnedB) && *A < *B); // Put pinned entries first
+	});
+}
+
+void SDrawPrimitiveDebugger::Refresh()
+{
+	AvailableEntries.Empty();
+	for (const TPair<FPrimitiveComponentId, FViewDebugInfo::FPrimitiveInfo>& Copy : LocalCopies)
+	{
+		AvailableEntries.Add(MakeShared<const FViewDebugInfo::FPrimitiveInfo>(Copy.Value)); // Add the local copies (pinned and hidden) at the top
+	}
+	FViewDebugInfo::Get().ForEachPrimitive([this](const FViewDebugInfo::FPrimitiveInfo& Primitive)
+	{
+		if (!LocalCopies.Contains(Primitive.ComponentId))
+		{
+			AvailableEntries.Add(MakeShared<const FViewDebugInfo::FPrimitiveInfo>(Primitive));
+			if (IsValid(Primitive.Component) && !Primitive.Component->IsVisible())
+			{
+				LocalCopies.Add(Primitive.ComponentId, Primitive);
+				HiddenEntries.Add(Primitive.ComponentId);
+			}
+		}
+	});
+	UpdateVisibleRows();
+}
+
+void SDrawPrimitiveDebugger::AddColumn(const FText& Name, const FName& ColumnId)
+{
+	const TSharedRef<FSlateFontMeasure> FontMeasure = FSlateApplication::Get().GetRenderer()->GetFontMeasureService();
+	const FSlateFontInfo FontInfo = FSlateFontInfo(FCoreStyle::GetDefaultFont(), 12);
+	const FName VisibilityColumn("Visible");
+	const FName PinColumn("Pin");
+	const FName NumDrawsColumn("NumDraws");
+	const FName LODColumn("LOD");
+	if (ColumnId.IsEqual(VisibilityColumn) || ColumnId.IsEqual(PinColumn) ||
+		ColumnId.IsEqual(NumDrawsColumn) || ColumnId.IsEqual(LODColumn))
+	{ // Rows that can be narrow and fixed
+		ColumnHeader->AddColumn(
+			SHeaderRow::Column(ColumnId)
+			.DefaultLabel(Name)
+			.FixedWidth(FontMeasure->Measure(Name, FontInfo).X)
+		);
+	}
+	else
+	{
+		ColumnHeader->AddColumn(
+			SHeaderRow::Column(ColumnId)
+			.DefaultLabel(Name)
+		);
+	}
+}
+
+void SDrawPrimitiveDebugger::OnChangeEntryVisibility(ECheckBoxState State, FPrimitiveRowDataPtr Data)
+{
+	if (Data.IsValid() && IsValid(Data->Component) && State != ECheckBoxState::Undetermined)
+	{
+		Data->Component->SetVisibility(State == ECheckBoxState::Checked);
+		if (State == ECheckBoxState::Unchecked)
+		{
+			if (!LocalCopies.Contains(Data->ComponentId))
+			{
+				LocalCopies.Add(Data->ComponentId, *Data); // Keep a local copy to prevent the primitive from disappearing during visibility check
+			}
+			HiddenEntries.Add(Data->ComponentId);
+		}
+		else if (State == ECheckBoxState::Checked)
+		{
+			HiddenEntries.Remove(Data->ComponentId);
+			if (!IsEntryPinned(Data))
+			{
+				LocalCopies.Remove(Data->ComponentId); // Remove the local copy when no longer needed
+			}
+		}
+	}
+}
+
+bool SDrawPrimitiveDebugger::IsEntryVisible(FPrimitiveRowDataPtr Data) const
+{
+	return !HiddenEntries.Contains(Data->ComponentId);
+}
+
+void SDrawPrimitiveDebugger::OnChangeEntryPinned(ECheckBoxState State, FPrimitiveRowDataPtr Data)
+{
+	if (State != ECheckBoxState::Undetermined)
+	{
+		if (State == ECheckBoxState::Checked)
+		{
+			if (!LocalCopies.Contains(Data->ComponentId))
+			{
+				LocalCopies.Add(Data->ComponentId, *Data); // Keep a local copy to prevent the primitive from disappearing during visibility check
+			}
+			PinnedEntries.Add(Data->ComponentId);
+		}
+		else if (State == ECheckBoxState::Unchecked)
+		{
+			PinnedEntries.Remove(Data->ComponentId);
+			if (IsEntryVisible(Data))
+			{
+				LocalCopies.Remove(Data->ComponentId); // Remove the local copy when no longer needed
+			}
+		}
+	}
+	UpdateVisibleRows();
+	if (Table.IsValid())
+	{
+		Table->RequestListRefresh();
+	}
+}
+
+bool SDrawPrimitiveDebugger::IsEntryPinned(FPrimitiveRowDataPtr Data) const
+{
+	return PinnedEntries.Contains(Data->ComponentId);
+}
+
+bool SDrawPrimitiveDebugger::CanCaptureSingleFrame() const
+{
+	return IDrawPrimitiveDebugger::IsAvailable() && !IDrawPrimitiveDebugger::Get().IsLiveCaptureEnabled();
+}
+
+FReply SDrawPrimitiveDebugger::OnRefreshClick()
+{
+	IDrawPrimitiveDebugger::Get().CaptureSingleFrame();
+	return FReply::Handled();
+}
+
+FReply SDrawPrimitiveDebugger::OnSaveClick()
+{
+	FViewDebugInfo::Get().DumpToCSV();
+	return FReply::Handled();
+}
+
+ECheckBoxState SDrawPrimitiveDebugger::IsLiveCaptureChecked() const
+{
+	return IDrawPrimitiveDebugger::Get().IsLiveCaptureEnabled() ? ECheckBoxState::Checked : ECheckBoxState::Unchecked;
+}
+
+void SDrawPrimitiveDebugger::OnToggleLiveCapture(ECheckBoxState state)
+{
+	if (state == ECheckBoxState::Checked)
+	{
+		IDrawPrimitiveDebugger::Get().EnableLiveCapture();
+	}
+	else if (state == ECheckBoxState::Unchecked)
+	{
+		IDrawPrimitiveDebugger::Get().DisableLiveCapture();
+	}
+}
+
+void SDrawPrimitiveDebuggerListViewRow::Construct(const FArguments& InArgs,
+                                             const TSharedRef<STableViewBase>& InOwnerTableView)
+{
+	RowDataPtr = InArgs._RowDataPtr;
+	DrawPrimitiveDebugger = InArgs._DrawPrimitiveDebugger;
+	SMultiColumnTableRow<FPrimitiveRowDataPtr>::Construct(
+		FSuperRowType::FArguments(),
+		InOwnerTableView
+	);
+}
+
+TSharedRef<SWidget> SDrawPrimitiveDebuggerListViewRow::GenerateWidgetForColumn(const FName& ColumnName)
+{
+	const TSharedPtr<SDrawPrimitiveDebugger> DrawPrimitiveDebuggerPtr = DrawPrimitiveDebugger.Pin();
+	return (DrawPrimitiveDebuggerPtr.IsValid())
+		? MakeCellWidget(IndexInList, ColumnName)
+		: SNullWidget::NullWidget;
+}
+
+TSharedRef<SWidget> SDrawPrimitiveDebuggerListViewRow::MakeCellWidget(const int32 InRowIndex, const FName& InColumnId)
+{
+	SDrawPrimitiveDebugger* DrawPrimitiveDebuggerPtr = DrawPrimitiveDebugger.Pin().Get();
+	if (DrawPrimitiveDebuggerPtr && RowDataPtr.IsValid())
+	{
+		const TSharedRef<FSlateFontMeasure> FontMeasure = FSlateApplication::Get().GetRenderer()->GetFontMeasureService();
+		const FSlateFontInfo FontInfo = FSlateFontInfo(FCoreStyle::GetDefaultFont(), UDrawPrimitiveDebuggerUserSettings::GetFontSize());
+		FText Value;
+		const FName VisibilityColumn("Visible");
+		const FName PinColumn("Pin");
+		const FName NameColumn("Name");
+		const FName ActorClassColumn("ActorClass");
+		const FName ActorColumn("Actor");
+		const FName LocationColumn("Location");
+		const FName MaterialsColumn("Materials");
+		const FName NumDrawsColumn("NumDraws");
+		const FName LODColumn("LOD");
+		const FName TrianglesColumn("Triangles");
+		if (InColumnId.IsEqual(VisibilityColumn))
+		{
+			return SNew(SBox)
+				.Padding(FMargin(5, 2, 5, 2))
+				.HAlign(HAlign_Center)
+				[
+					SNew(SCheckBox)
+						.IsChecked(this, &SDrawPrimitiveDebuggerListViewRow::IsVisible)
+						.OnCheckStateChanged(DrawPrimitiveDebuggerPtr, &SDrawPrimitiveDebugger::OnChangeEntryVisibility, RowDataPtr)
+						.HAlign(HAlign_Center)
+				];
+		}
+		if (InColumnId.IsEqual(PinColumn))
+		{
+			return SNew(SBox)
+				.Padding(FMargin(5, 2, 5, 2))
+				.HAlign(HAlign_Center)
+				[
+					SNew(SCheckBox)
+						.IsChecked(this, &SDrawPrimitiveDebuggerListViewRow::IsPinned)
+						.OnCheckStateChanged(DrawPrimitiveDebuggerPtr, &SDrawPrimitiveDebugger::OnChangeEntryPinned, RowDataPtr)
+						.HAlign(HAlign_Center)
+				];
+		}
+		if (InColumnId.IsEqual(NameColumn))
+		{
+			Value = FText::FromString(RowDataPtr->Name);
+		}
+		else if (InColumnId.IsEqual(ActorClassColumn))
+		{
+			Value = IsValid(RowDataPtr->Owner) ?
+				FText::FromString(RowDataPtr->Owner->GetClass()->GetName()) :
+				FText::FromString("INVALID");
+		}
+		else if (InColumnId.IsEqual(ActorColumn))
+		{
+			Value = IsValid(RowDataPtr->Owner) ?
+				FText::FromString(RowDataPtr->Owner->GetHumanReadableName()) :
+				FText::FromString("INVALID");
+		}
+		else if (InColumnId.IsEqual(LocationColumn))
+		{
+			Value = IsValid(RowDataPtr->Owner) ?
+				FText::FromString(RowDataPtr->Owner->GetActorLocation().ToString()) :
+				FText::FromString("INVALID");
+		}
+		else if (InColumnId.IsEqual(MaterialsColumn))
+		{
+			const int32 Count = RowDataPtr->Materials.Num();
+			FString Materials = FString::Printf(TEXT("[%d] {"), Count);
+			for (int i = 0; i < Count; i++)
+			{
+				const UMaterialInterface* MI = RowDataPtr->Materials[i];
+				if (MI && MI->GetMaterial())
+				{
+					Materials += MI->GetMaterial()->GetName();
+				}
+				else
+				{
+					Materials += "Null";
+				}
+				
+				if (i < Count - 1)
+				{
+					Materials += ", ";
+				}
+			}
+			Materials += "}";
+			Value = FText::FromString(Materials);
+		}
+		else if (InColumnId.IsEqual(NumDrawsColumn))
+		{
+			Value = FText::FromString(FString::FromInt(RowDataPtr->DrawCount));
+		}
+		else if (InColumnId.IsEqual(LODColumn))
+		{
+			Value = FText::FromString(FString::FromInt(RowDataPtr->LOD));
+		}
+		else if (InColumnId.IsEqual(TrianglesColumn))
+		{
+			Value = FText::FromString(FString::FromInt(RowDataPtr->TriangleCount));
+		}
+		else
+		{
+			// Invalid Column name
+			return SNullWidget::NullWidget;
+		}
+		return SNew(SBox)
+			.Padding(FMargin(5, 2, 5, 2))
+			.HAlign(HAlign_Fill)
+			[
+				SNew(STextBlock)
+				.ColorAndOpacity(FSlateColor::UseForeground())
+				.Text(Value)
+				.Font(FontInfo)
+				.IsEnabled(DrawPrimitiveDebuggerPtr, &SDrawPrimitiveDebugger::IsEntryVisible, RowDataPtr)
+				.Justification(ETextJustify::Left)
+				.HighlightText(DrawPrimitiveDebuggerPtr, &SDrawPrimitiveDebugger::GetFilterText)
+			];
+	}
+	return SNullWidget::NullWidget;
+}
+
+ECheckBoxState SDrawPrimitiveDebuggerListViewRow::IsVisible() const
+{
+	const SDrawPrimitiveDebugger* DrawPrimitiveDebuggerPtr = DrawPrimitiveDebugger.Pin().Get();
+	return DrawPrimitiveDebuggerPtr && DrawPrimitiveDebuggerPtr->IsEntryVisible(RowDataPtr) ? ECheckBoxState::Checked : ECheckBoxState::Unchecked;
+}
+
+ECheckBoxState SDrawPrimitiveDebuggerListViewRow::IsPinned() const
+{
+	const SDrawPrimitiveDebugger* DrawPrimitiveDebuggerPtr = DrawPrimitiveDebugger.Pin().Get();
+	return DrawPrimitiveDebuggerPtr && DrawPrimitiveDebuggerPtr->IsEntryPinned(RowDataPtr) ? ECheckBoxState::Checked : ECheckBoxState::Unchecked;
+}
+
+END_SLATE_FUNCTION_BUILD_OPTIMIZATION
+
+#endif
\ No newline at end of file
diff --git a/Engine/Source/Developer/DrawPrimitiveDebugger/Private/SDrawPrimitiveDebugger.h b/Engine/Source/Developer/DrawPrimitiveDebugger/Private/SDrawPrimitiveDebugger.h
new file mode 100644
--- /dev/null
+++ b/Engine/Source/Developer/DrawPrimitiveDebugger/Private/SDrawPrimitiveDebugger.h
@@ -0,0 +1,90 @@
+﻿// Copyright Epic Games, Inc. All Rights Reserved.
+
+#pragma once
+
+#include "CoreMinimal.h"
+#include "Widgets/SCompoundWidget.h"
+#include "Slate/Public/Widgets/Input/SSearchBox.h"
+#include "ViewDebug.h"
+
+#if !UE_BUILD_SHIPPING
+
+typedef TSharedPtr<const FViewDebugInfo::FPrimitiveInfo> FPrimitiveRowDataPtr;
+
+class SDrawPrimitiveDebugger : public SCompoundWidget
+{
+public:
+	SLATE_BEGIN_ARGS(SDrawPrimitiveDebugger)
+		{
+		}
+
+	SLATE_END_ARGS()
+
+	/** Constructs this widget with InArgs */
+	void Construct(const FArguments& InArgs);
+	FText GetFilterText() const;
+	void OnFilterTextChanged(const FText& InFilterText);
+	void OnFilterTextCommitted(const FText& NewText, ETextCommit::Type CommitInfo);
+	TSharedRef<ITableRow> MakeRowWidget(FPrimitiveRowDataPtr InRowDataPtr, const TSharedRef<STableViewBase>& OwnerTable);
+	void UpdateVisibleRows();
+	void SortRows();
+	void Refresh();
+	void AddColumn(const FText& Name, const FName& ColumnId);
+	void OnChangeEntryVisibility(ECheckBoxState state, FPrimitiveRowDataPtr data);
+	bool IsEntryVisible(FPrimitiveRowDataPtr data) const;
+	void OnChangeEntryPinned(ECheckBoxState state, FPrimitiveRowDataPtr data);
+	bool IsEntryPinned(FPrimitiveRowDataPtr data) const;
+	bool CanCaptureSingleFrame() const;
+	FReply OnRefreshClick();
+	FReply OnSaveClick();
+	ECheckBoxState IsLiveCaptureChecked() const;
+	void OnToggleLiveCapture(ECheckBoxState state);
+
+private:
+	
+	TSharedPtr<SSearchBox> SearchBox;
+	TSharedPtr<SListView<FPrimitiveRowDataPtr>> Table;
+	FText FilterText;
+	TSharedPtr<SHeaderRow> ColumnHeader;
+	TArray<FPrimitiveRowDataPtr> AvailableEntries;
+	TArray<FPrimitiveRowDataPtr> VisibleEntries;
+
+	// TODO: Decide if primitive data copies are necessary
+	TMap<FPrimitiveComponentId, FViewDebugInfo::FPrimitiveInfo> LocalCopies;
+	TSet<FPrimitiveComponentId> PinnedEntries;
+	TSet<FPrimitiveComponentId> HiddenEntries;
+};
+
+/**
+ * A widget to represent a row in a Data Table Editor widget. This widget allows us to do things like right-click
+ * and take actions on a particular row of a Data Table.
+ */
+class SDrawPrimitiveDebuggerListViewRow : public SMultiColumnTableRow<FPrimitiveRowDataPtr>
+{
+public:
+
+	SLATE_BEGIN_ARGS(SDrawPrimitiveDebuggerListViewRow)
+	{
+	}
+	/** The owning object. This allows us access to the actual data table being edited as well as some other API functions. */
+	SLATE_ARGUMENT(TWeakPtr<SDrawPrimitiveDebugger>, DrawPrimitiveDebugger)
+		/** The primitive we're working with to allow us to get naming information. */
+		SLATE_ARGUMENT(FPrimitiveRowDataPtr, RowDataPtr)
+		SLATE_END_ARGS()
+
+	void Construct(const FArguments& InArgs, const TSharedRef<STableViewBase>& InOwnerTableView);
+
+	virtual TSharedRef<SWidget> GenerateWidgetForColumn(const FName& ColumnName) override;
+
+private:
+
+	TSharedRef<SWidget> MakeCellWidget(const int32 InRowIndex, const FName& InColumnId);
+	
+	ECheckBoxState IsVisible() const;
+	ECheckBoxState IsPinned() const;
+
+	FPrimitiveRowDataPtr RowDataPtr;
+	TWeakPtr<SDrawPrimitiveDebugger> DrawPrimitiveDebugger;
+};
+
+#endif
\ No newline at end of file
diff --git a/Engine/Source/Developer/DrawPrimitiveDebugger/Public/DrawPrimitiveDebugger.h b/Engine/Source/Developer/DrawPrimitiveDebugger/Public/DrawPrimitiveDebugger.h
new file mode 100644
--- /dev/null
+++ b/Engine/Source/Developer/DrawPrimitiveDebugger/Public/DrawPrimitiveDebugger.h
@@ -0,0 +1,70 @@
+﻿// Copyright Epic Games, Inc. All Rights Reserved.
+
+// DRAW PRIMITIVE DEBUGGER
+// 
+// This tool allows easy on screen debugging of graphics data, primarily on-screen primitives
+// 
+
+#pragma once
+
+#include "CoreMinimal.h"
+#include "Modules/ModuleInterface.h"
+#include "Modules/ModuleManager.h"
+
+DECLARE_LOG_CATEGORY_EXTERN(LogDrawPrimitiveDebugger, All, All);
+
+class DRAWPRIMITIVEDEBUGGER_API IDrawPrimitiveDebugger : public IModuleInterface
+{
+
+public:
+	/**
+	 * Singleton-like access to this module's interface. This is just for convenience!
+	 * Beware of calling this during the shutdown phase, though. Your module might have been unloaded already.
+	 *
+	 * @return Returns singleton instance, loading the module on demand if needed
+	 */
+	static inline IDrawPrimitiveDebugger& Get()
+	{
+		return FModuleManager::LoadModuleChecked< IDrawPrimitiveDebugger >("DrawPrimitiveDebugger");
+	}
+
+	/**
+	 * Checks to see if this module is loaded and ready. It is only valid to call Get() if IsAvailable() returns true.
+	 *
+	 * @return True if the module is loaded and ready to use
+	 */
+	static inline bool IsAvailable()
+	{
+		return FModuleManager::Get().IsModuleLoaded("DrawPrimitiveDebugger");
+	}
+
+	/**
+	 * Instructs the renderer to capture a snapshot for the debugger on the next frame.
+	 */
+	virtual void CaptureSingleFrame() = 0;
+
+	/**
+	 * Is live data capture enabled for the debugger? This implies that the renderer will capture debug data each frame.
+	 * @return True if live capture is enabled.
+	 */
+	virtual bool IsLiveCaptureEnabled() const = 0;
+	
+	/**
+	 * Enables capturing debug data each frame.
+	 */
+	virtual void EnableLiveCapture() = 0;
+	/**
+	 * Disables capturing debug data each frame.
+	 */
+	virtual void DisableLiveCapture() = 0;
+
+	/**
+	 * Opens the graphics debugger window if it is available.
+	 */
+	virtual void OpenDebugWindow() = 0;
+	/**
+	 * Closes the graphics debugger window if it is currently open.
+	 */
+	virtual void CloseDebugWindow() = 0;
+	
+};
\ No newline at end of file
diff --git a/Engine/Source/Developer/DrawPrimitiveDebugger/Public/DrawPrimitiveDebuggerConfig.h b/Engine/Source/Developer/DrawPrimitiveDebugger/Public/DrawPrimitiveDebuggerConfig.h
new file mode 100644
--- /dev/null
+++ b/Engine/Source/Developer/DrawPrimitiveDebugger/Public/DrawPrimitiveDebuggerConfig.h
@@ -0,0 +1,26 @@
+﻿// Copyright Epic Games, Inc. All Rights Reserved.
+
+#pragma once
+
+#include "CoreMinimal.h"
+#include "Engine/DeveloperSettings.h"
+#include "DrawPrimitiveDebuggerConfig.generated.h"
+
+UCLASS(config = EditorPerProjectUserSettings, meta = (DisplayName = "Draw Primitive Debugger"))
+class DRAWPRIMITIVEDEBUGGER_API UDrawPrimitiveDebuggerUserSettings : public UDeveloperSettings
+{
+	GENERATED_BODY()
+protected:
+	virtual FName GetCategoryName() const override { return TEXT("Advanced"); }
+
+public:
+	
+	static int32 GetFontSize() { return GetDefault<UDrawPrimitiveDebuggerUserSettings>()->FontSize; }
+	static void SetFontSize(const int32 InFontSize);
+
+protected:
+	
+	/** Font Size used by Draw Primitive Debugger */ 
+	UPROPERTY(config, EditAnywhere, Category = DrawPrimitiveDebugger)
+	int32 FontSize = 10;
+};
\ No newline at end of file
diff --git a/Engine/Source/Developer/GameplayDebugger/GameplayDebugger.Build.cs b/Engine/Source/Developer/GameplayDebugger/GameplayDebugger.Build.cs
--- a/Engine/Source/Developer/GameplayDebugger/GameplayDebugger.Build.cs
+++ b/Engine/Source/Developer/GameplayDebugger/GameplayDebugger.Build.cs
@@ -21,6 +21,11 @@ namespace UnrealBuildTool.Rules
 					"Slate",
 					"DeveloperSettings",
 				});
+			
+			if (Target.Configuration != UnrealTargetConfiguration.Shipping)
+			{
+				PrivateDependencyModuleNames.Add("DrawPrimitiveDebugger");
+			}
 
 			PrivateIncludePaths.AddRange(
 				new string[] {
diff --git a/Engine/Source/Runtime/Renderer/Private/SceneRendering.h b/Engine/Source/Runtime/Renderer/Private/SceneRendering.h
--- a/Engine/Source/Runtime/Renderer/Private/SceneRendering.h
+++ b/Engine/Source/Runtime/Renderer/Private/SceneRendering.h
@@ -2369,6 +2369,8 @@ private:
 	void ComputeFamilySize();
 
 #if !UE_BUILD_SHIPPING
+	/** Collect the draw data of all visible UPrimitiveComponents in the Scene */
+	void ProcessPrimitives(const FViewInfo& View, const FViewCommands& ViewCommands) const;
 	/** Dump all UPrimitiveComponents in the Scene to a CSV file */
 	void DumpPrimitives(const FViewCommands& ViewCommands);
 #endif
diff --git a/Engine/Source/Runtime/Renderer/Private/SceneVisibility.cpp b/Engine/Source/Runtime/Renderer/Private/SceneVisibility.cpp
--- a/Engine/Source/Runtime/Renderer/Private/SceneVisibility.cpp
+++ b/Engine/Source/Runtime/Renderer/Private/SceneVisibility.cpp
@@ -47,6 +47,7 @@
 #include "RayTracing/RayTracingInstanceCulling.h"
 #include "RendererModule.h"
 #include "SceneViewExtension.h"
+#include "ViewDebug.h"
 
 /*------------------------------------------------------------------------------
 	Globals
@@ -295,6 +296,7 @@ static TAutoConsoleVariable<float> CVarFreezeTemporalHistories(
 	ECVF_RenderThreadSafe);
 
 static bool bDumpPrimitivesNextFrame = false;
+static bool bDumpDetailedPrimitivesNextFrame = false;
 
 static FAutoConsoleCommand CVarDumpPrimitives(
 	TEXT("DumpPrimitives"),
@@ -302,6 +304,12 @@ static FAutoConsoleCommand CVarDumpPrimitives(
 	FConsoleCommandDelegate::CreateStatic([] { bDumpPrimitivesNextFrame = true; }),
 	ECVF_Default);
 
+static FAutoConsoleCommand CVarDrawPrimitiveDebugData(
+	TEXT("DumpDetailedPrimitives"),
+	TEXT("Writes out all scene primitive details to a CSV file"),
+	FConsoleCommandDelegate::CreateStatic([] { bDumpDetailedPrimitivesNextFrame = !bDumpDetailedPrimitivesNextFrame; }),
+	ECVF_Default);
+
 #endif
 
 DECLARE_CYCLE_STAT(TEXT("Occlusion Readback"), STAT_CLMM_OcclusionReadback, STATGROUP_CommandListMarkers);
@@ -4204,6 +4212,211 @@ static uint32 GetDrawCountFromPrimitiveSceneInfo(FScene* Scene, const FPrimitive
 	return DrawCount;
 }
 
+
+FViewDebugInfo FViewDebugInfo::Instance;
+
+FViewDebugInfo::FViewDebugInfo()
+{
+	bHasEverUpdated = false;
+	bIsOutdated = true;
+	bShouldUpdate = false;
+	bShouldCaptureSingleFrame = false;
+}
+
+/*bool FViewDebugInfo::FPrimitiveInfo::operator<(const FPrimitiveInfo& Other) const
+{
+	// Sort by name to group similar assets together, then by exact primitives so we can ignore duplicates
+	const int32 NameCompare = Name.Compare(Other.Name);
+	if (NameCompare != 0)
+	{
+		return NameCompare < 0;
+	}
+
+	return PrimitiveSceneInfo < Other.PrimitiveSceneInfo;
+}*/
+
+void FViewDebugInfo::ProcessPrimitive(FPrimitiveSceneInfo* PrimitiveSceneInfo, const FViewInfo& View, FScene* Scene, const UPrimitiveComponent* DebugComponent)
+{
+	if (!DebugComponent->IsRegistered())
+	{
+		return;
+	}
+	AActor* Actor = DebugComponent->GetOwner();
+	FString FullName = DebugComponent->GetName();
+	const uint32 DrawCount = GetDrawCountFromPrimitiveSceneInfo(Scene, PrimitiveSceneInfo);
+
+	TArray<UMaterialInterface*> Materials;
+	DebugComponent->GetUsedMaterials(Materials);
+	const int32 LOD = PrimitiveSceneInfo->Proxy ? PrimitiveSceneInfo->Proxy->GetLOD(&View) : INDEX_NONE;
+	int32 Triangles = 0;
+	if (const UStaticMeshComponent* StaticMeshComponent = Cast<UStaticMeshComponent>(DebugComponent))
+	{
+		Triangles = StaticMeshComponent->GetStaticMesh()->GetNumTriangles(LOD);
+	}
+	else if (const USkeletalMeshComponent* SkeletalMeshComponent = Cast<USkeletalMeshComponent>(DebugComponent))
+	{
+		for (const FSkeletalMeshLODRenderData& RenderData : SkeletalMeshComponent->GetSkeletalMeshRenderData()->LODRenderData)
+		{
+			Triangles += RenderData.MultiSizeIndexContainer.GetIndexBuffer()->Num() / 3;
+		}
+	}
+		
+	const FPrimitiveInfo PrimitiveInfo = {
+		Actor,
+		DebugComponent->ComponentId,
+		const_cast<UPrimitiveComponent*>(DebugComponent), // This is probably a bad idea, find alternative
+		PrimitiveSceneInfo,
+		MoveTemp(Materials),
+		MoveTemp(FullName),
+		DrawCount,
+		Triangles,
+		LOD
+	};
+
+	Primitives.Add(PrimitiveInfo);
+}
+
+void FViewDebugInfo::DumpToCSV() const
+{
+	const FString OutputPath = FPaths::ProfilingDir() / TEXT("Primitives") / FString::Printf(TEXT("PrimitivesDetailed-%s.csv"), *FDateTime::Now().ToString());
+	const bool bSuppressViewer = true;
+	FDiagnosticTableViewer DrawViewer(*OutputPath, bSuppressViewer);
+	DrawViewer.AddColumn(TEXT("Name"));
+	DrawViewer.AddColumn(TEXT("ActorClass"));
+	DrawViewer.AddColumn(TEXT("Actor"));
+	DrawViewer.AddColumn(TEXT("Location"));
+	DrawViewer.AddColumn(TEXT("NumMaterials"));
+	DrawViewer.AddColumn(TEXT("Materials"));
+	DrawViewer.AddColumn(TEXT("NumDraws"));
+	DrawViewer.AddColumn(TEXT("LOD"));
+	DrawViewer.AddColumn(TEXT("Triangles"));
+	DrawViewer.CycleRow();
+
+	FRWScopeLock ScopeLock(Lock, SLT_ReadOnly);
+	const FPrimitiveSceneInfo* LastPrimitiveSceneInfo = nullptr;
+	for (const FPrimitiveInfo& Primitive : Primitives)
+	{
+		if (Primitive.PrimitiveSceneInfo != LastPrimitiveSceneInfo)
+		{
+			DrawViewer.AddColumn(*Primitive.Name);
+			DrawViewer.AddColumn(Primitive.Owner ? *Primitive.Owner->GetClass()->GetName() : TEXT(""));
+			DrawViewer.AddColumn(Primitive.Owner ? *Primitive.Owner->GetFullName() : TEXT(""));
+			DrawViewer.AddColumn(Primitive.Owner ?
+				*FString::Printf(TEXT("{%s}"), *Primitive.Owner->GetActorLocation().ToString()) : TEXT(""));
+			DrawViewer.AddColumn(*FString::Printf(TEXT("%d"), Primitive.Materials.Num()));
+			FString Materials = "[";
+			for (int i = 0; i < Primitive.Materials.Num(); i++)
+			{
+				if (Primitive.Materials[i] && Primitive.Materials[i]->GetMaterial())
+				{
+					Materials += Primitive.Materials[i]->GetMaterial()->GetName();
+				}
+				else
+				{
+					Materials += "Null";
+				}
+				
+				if (i < Primitive.Materials.Num() - 1)
+				{
+					Materials += ", ";
+				}
+			}
+			Materials += "]";
+			DrawViewer.AddColumn(*FString::Printf(TEXT("%s"), *Materials));
+			DrawViewer.AddColumn(*FString::Printf(TEXT("%d"), Primitive.DrawCount));
+			DrawViewer.AddColumn(*FString::Printf(TEXT("%d"), Primitive.LOD));
+			DrawViewer.AddColumn(*FString::Printf(TEXT("%d"), Primitive.TriangleCount));
+			DrawViewer.CycleRow();
+
+			LastPrimitiveSceneInfo = Primitive.PrimitiveSceneInfo;
+		}
+	}
+}
+
+void FViewDebugInfo::CaptureNextFrame()
+{
+	FRWScopeLock ScopeLock(Lock, SLT_Write);
+	bShouldCaptureSingleFrame = true;
+	bShouldUpdate = true;
+}
+
+void FViewDebugInfo::EnableLiveCapture()
+{
+	FRWScopeLock ScopeLock(Lock, SLT_Write);
+	bShouldCaptureSingleFrame = false;
+	bShouldUpdate = true;
+}
+
+void FViewDebugInfo::DisableLiveCapture()
+{
+	FRWScopeLock ScopeLock(Lock, SLT_Write);
+	bShouldCaptureSingleFrame = false;
+	bShouldUpdate = false;
+}
+
+bool FViewDebugInfo::HasEverUpdated() const
+{
+	FRWScopeLock ScopeLock(Lock, SLT_ReadOnly);
+	return bHasEverUpdated;
+}
+
+bool FViewDebugInfo::IsOutOfDate() const
+{
+	FRWScopeLock ScopeLock(Lock, SLT_ReadOnly);
+	return bIsOutdated;
+}
+
+void FSceneRenderer::ProcessPrimitives(const FViewInfo& View, const FViewCommands& ViewCommands) const
+{
+	FViewDebugInfo& DebugInfo = FViewDebugInfo::Instance;
+	{
+		FRWScopeLock ScopeLock(DebugInfo.Lock, SLT_Write);
+		DebugInfo.bIsOutdated = true;
+	
+		if (!DebugInfo.bShouldUpdate && !bDumpDetailedPrimitivesNextFrame)
+		{
+			return;
+		}
+
+		if (DebugInfo.bShouldCaptureSingleFrame)
+		{
+			DebugInfo.bShouldCaptureSingleFrame = false;
+			DebugInfo.bShouldUpdate = false;
+		}
+
+		// TODO: Add profiling to this function
+	
+		DebugInfo.Primitives.Empty(ViewCommands.MeshCommands[EMeshPass::BasePass].Num() + ViewCommands.DynamicMeshCommandBuildRequests[EMeshPass::BasePass].Num());
+
+		for (const FVisibleMeshDrawCommand& Mesh : ViewCommands.MeshCommands[EMeshPass::BasePass])
+		{
+			const int32 PrimitiveId = Mesh.PrimitiveIdInfo.ScenePrimitiveId;
+			if (PrimitiveId >= 0 && PrimitiveId < Scene->Primitives.Num())
+			{
+				FPrimitiveSceneInfo* PrimitiveSceneInfo = Scene->Primitives[PrimitiveId];
+				DebugInfo.ProcessPrimitive(PrimitiveSceneInfo, View, Scene, PrimitiveSceneInfo->ComponentForDebuggingOnly);
+			}
+		}
+
+		for (const FStaticMeshBatch* StaticMeshBatch : ViewCommands.DynamicMeshCommandBuildRequests[EMeshPass::BasePass])
+		{
+			FPrimitiveSceneInfo* PrimitiveSceneInfo = StaticMeshBatch->PrimitiveSceneInfo;
+			DebugInfo.ProcessPrimitive(PrimitiveSceneInfo, View, Scene, PrimitiveSceneInfo->ComponentForDebuggingOnly);
+		}
+
+		DebugInfo.bHasEverUpdated = true;
+		DebugInfo.bIsOutdated = false;
+	}
+	DebugInfo.OnUpdate.Broadcast();
+	
+	if (bDumpDetailedPrimitivesNextFrame)
+	{
+		DebugInfo.DumpToCSV();
+		bDumpDetailedPrimitivesNextFrame = false;
+	}
+	
+}
+
 void FSceneRenderer::DumpPrimitives(const FViewCommands& ViewCommands)
 {
 	if (!bDumpPrimitivesNextFrame)
@@ -4671,6 +4884,7 @@ void FSceneRenderer::ComputeViewVisibility(
 
 #if !UE_BUILD_SHIPPING
 		DumpPrimitives(ViewCommands);
+		ProcessPrimitives(View, ViewCommands);
 #endif
 
 		SetupMeshPass(View, BasePassDepthStencilAccess, ViewCommands, InstanceCullingManager);
diff --git a/Engine/Source/Runtime/Renderer/Public/ViewDebug.h b/Engine/Source/Runtime/Renderer/Public/ViewDebug.h
new file mode 100644
--- /dev/null
+++ b/Engine/Source/Runtime/Renderer/Public/ViewDebug.h
@@ -0,0 +1,140 @@
+﻿// Copyright Epic Games, Inc. All Rights Reserved.
+
+#pragma once
+
+#include "CoreMinimal.h"
+#include "PrimitiveSceneInfo.h"
+
+#if !UE_BUILD_SHIPPING // TODO: Decide whether or not the struct should be entirely stripped out of shipping
+
+DECLARE_MULTICAST_DELEGATE(FOnUpdateViewDebugInfo);
+
+/**
+ * A collection of debug data associated with the current on screen view.
+ */
+struct RENDERER_API FViewDebugInfo
+{
+	friend class FSceneRenderer;
+	friend class FDrawPrimitiveDebuggerModule;
+
+private:
+
+	static FViewDebugInfo Instance;
+
+	FViewDebugInfo();
+	
+public:
+
+	/**
+	 * Gets a reference to the view debug information that is used by the renderer.
+	 * @returns The debug information that is used by the renderer.
+	 */
+	static inline const FViewDebugInfo& Get()
+	{
+		return Instance;
+	}
+
+	/**
+	 * Data collected about a single primitive being drawn to the screen.
+	 */
+	struct FPrimitiveInfo
+	{
+		AActor* Owner;
+		FPrimitiveComponentId ComponentId;
+		UPrimitiveComponent* Component;
+		FPrimitiveSceneInfo* PrimitiveSceneInfo;
+		TArray<UMaterialInterface*> Materials;
+		FString Name;
+		uint32 DrawCount;
+		int32 TriangleCount;
+		int32 LOD;
+
+		bool operator<(const FPrimitiveInfo& Other) const
+		{
+			// Sort by name to group similar assets together, then by exact primitives so we can ignore duplicates
+			const int32 NameCompare = Name.Compare(Other.Name);
+			if (NameCompare != 0)
+			{
+				return NameCompare < 0;
+			}
+
+			return PrimitiveSceneInfo < Other.PrimitiveSceneInfo;
+		}
+	};
+
+private:
+
+	bool bHasEverUpdated;
+	bool bIsOutdated;
+	bool bShouldUpdate;
+	bool bShouldCaptureSingleFrame;
+
+	FOnUpdateViewDebugInfo OnUpdate;
+
+	mutable FRWLock Lock;
+	
+	TArray<FPrimitiveInfo> Primitives;
+
+	void ProcessPrimitive(FPrimitiveSceneInfo* PrimitiveSceneInfo, const FViewInfo& View, FScene* Scene, const UPrimitiveComponent* DebugComponent);
+
+	void CaptureNextFrame();
+
+	void EnableLiveCapture();
+
+	void DisableLiveCapture();
+
+public:
+
+	/**
+	 * Writes the currently stored information out to a CSV file.
+	 */
+	void DumpToCSV() const;
+
+	/**
+	 * Performs an operation for each primitive currently tracked.
+	 * @param Action The action to perform for each primitive.
+	 */
+	template <typename CallableT>
+	void ForEachPrimitive(CallableT Action) const
+	{
+		const FPrimitiveSceneInfo* LastPrimitiveSceneInfo = nullptr;
+		FRWScopeLock ScopeLock(Lock, SLT_ReadOnly);
+		for (const FPrimitiveInfo& Primitive : Primitives)
+		{
+			if (Primitive.PrimitiveSceneInfo != LastPrimitiveSceneInfo)
+			{
+				Invoke(Action, Primitive);
+				LastPrimitiveSceneInfo = Primitive.PrimitiveSceneInfo;
+			}
+		}
+	}
+
+	/**
+	 * Checks if this debug information has ever been updated.
+	 * @returns True if the information has been updated at least once.
+	 */
+	bool HasEverUpdated() const;
+
+	/**
+	 * Checks if current information is from an older frame.
+	 * @returns True if the data in this object is outdated.
+	 */
+	bool IsOutOfDate() const;
+
+	template <typename UserClass>
+	FDelegateHandle AddUpdateHandler(UserClass* UserObject, void (UserClass::*Func)())
+	{
+		return OnUpdate.AddRaw(UserObject, Func);
+	}
+
+	FDelegateHandle AddUpdateHandler(void (*Func)())
+	{
+		return OnUpdate.AddStatic(Func);
+	}
+
+	void RemoveUpdateHandler(const FDelegateHandle& Handle)
+	{
+		OnUpdate.Remove(Handle);
+	}
+};
+#endif
\ No newline at end of file
