Fixing missing header includes in DrawPrimitiveDebugger slate implementation.

diff --git a/Engine/Source/Developer/DrawPrimitiveDebugger/Private/SDrawPrimitiveDebugger.cpp b/Engine/Source/Developer/DrawPrimitiveDebugger/Private/SDrawPrimitiveDebugger.cpp
--- a/Engine/Source/Developer/DrawPrimitiveDebugger/Private/SDrawPrimitiveDebugger.cpp
+++ b/Engine/Source/Developer/DrawPrimitiveDebugger/Private/SDrawPrimitiveDebugger.cpp
@@ -6,7 +6,9 @@
 #include "DrawPrimitiveDebuggerConfig.h"
 #include "SlateOptMacros.h"
 #include "Fonts/FontMeasure.h"
-#include "Widgets/Images/SImage.h"
+#include "Slate.h"
+#include "Widgets/Input/SButton.h"
+#include "Widgets/Input/SSearchBox.h"
 #include "Widgets/Layout/SScrollBox.h"
 
 #if !UE_BUILD_SHIPPING
