Replacing monolithic header include from DrawPrimitiveDebugger slate implementation.

diff --git a/Engine/Source/Developer/DrawPrimitiveDebugger/Private/SDrawPrimitiveDebugger.cpp b/Engine/Source/Developer/DrawPrimitiveDebugger/Private/SDrawPrimitiveDebugger.cpp
--- a/Engine/Source/Developer/DrawPrimitiveDebugger/Private/SDrawPrimitiveDebugger.cpp
+++ b/Engine/Source/Developer/DrawPrimitiveDebugger/Private/SDrawPrimitiveDebugger.cpp
@@ -6,10 +6,16 @@
 #include "DrawPrimitiveDebuggerConfig.h"
 #include "SlateOptMacros.h"
 #include "Fonts/FontMeasure.h"
-#include "Slate.h"
 #include "Widgets/Input/SButton.h"
+#include "Widgets/Views/SListView.h"
+#include "Widgets/Views/SHeaderRow.h"
+#include "Widgets/SBoxPanel.h"
+#include "Widgets/Input/SCheckBox.h"
 #include "Widgets/Input/SSearchBox.h"
 #include "Widgets/Layout/SScrollBox.h"
+#include "Widgets/Layout/SScrollBar.h"
+#include "Widgets/Layout/SBox.h"
+#include "Widgets/Text/STextBlock.h"
 
 #if !UE_BUILD_SHIPPING
 
