CIS Nonunity fix, not my code but I broke it at the same time

diff --git a/Engine/Source/Developer/DrawPrimitiveDebugger/Private/SDrawPrimitiveDebugger.cpp b/Engine/Source/Developer/DrawPrimitiveDebugger/Private/SDrawPrimitiveDebugger.cpp
--- a/Engine/Source/Developer/DrawPrimitiveDebugger/Private/SDrawPrimitiveDebugger.cpp
+++ b/Engine/Source/Developer/DrawPrimitiveDebugger/Private/SDrawPrimitiveDebugger.cpp
@@ -1,11 +1,13 @@
-﻿// Copyright Epic Games, Inc. All Rights Reserved.
+// Copyright Epic Games, Inc. All Rights Reserved.
 
 #include "SDrawPrimitiveDebugger.h"
 
 #include "DrawPrimitiveDebugger.h"
 #include "DrawPrimitiveDebuggerConfig.h"
 #include "SlateOptMacros.h"
+#include "Components/PrimitiveComponent.h"
 #include "Fonts/FontMeasure.h"
+#include "Materials/Material.h"
 #include "Widgets/Input/SButton.h"
 #include "Widgets/Views/SListView.h"
 #include "Widgets/Views/SHeaderRow.h"
