Fix Unity compile missing headers in SDrawPrimitiveDebugger

diff --git a/Engine/Source/Developer/DrawPrimitiveDebugger/Private/SDrawPrimitiveDebugger.cpp b/Engine/Source/Developer/DrawPrimitiveDebugger/Private/SDrawPrimitiveDebugger.cpp
--- a/Engine/Source/Developer/DrawPrimitiveDebugger/Private/SDrawPrimitiveDebugger.cpp
+++ b/Engine/Source/Developer/DrawPrimitiveDebugger/Private/SDrawPrimitiveDebugger.cpp
@@ -5,15 +5,13 @@
 #include "DrawPrimitiveDebugger.h"
 #include "DrawPrimitiveDebuggerConfig.h"
 #include "SlateOptMacros.h"
+#include "Framework/Application/SlateApplication.h"
 #include "Components/PrimitiveComponent.h"
 #include "Fonts/FontMeasure.h"
 #include "Materials/Material.h"
 #include "Widgets/Input/SButton.h"
-#include "Widgets/Views/SListView.h"
-#include "Widgets/Views/SHeaderRow.h"
 #include "Widgets/SBoxPanel.h"
 #include "Widgets/Input/SCheckBox.h"
-#include "Widgets/Input/SSearchBox.h"
 #include "Widgets/Layout/SScrollBox.h"
 #include "Widgets/Layout/SScrollBar.h"
 #include "Widgets/Layout/SBox.h"
diff --git a/Engine/Source/Developer/DrawPrimitiveDebugger/Private/SDrawPrimitiveDebugger.h b/Engine/Source/Developer/DrawPrimitiveDebugger/Private/SDrawPrimitiveDebugger.h
--- a/Engine/Source/Developer/DrawPrimitiveDebugger/Private/SDrawPrimitiveDebugger.h
+++ b/Engine/Source/Developer/DrawPrimitiveDebugger/Private/SDrawPrimitiveDebugger.h
@@ -4,7 +4,9 @@
 
 #include "CoreMinimal.h"
 #include "Widgets/SCompoundWidget.h"
-#include "Slate/Public/Widgets/Input/SSearchBox.h"
+#include "Widgets/Views/SHeaderRow.h"
+#include "Widgets/Views/SListView.h"
+#include "Widgets/Input/SSearchBox.h"
 #include "ViewDebug.h"
 
 #if !UE_BUILD_SHIPPING
