Fixing search bar material filtering functionality for DrawPrimitiveDebugger

diff --git a/Engine/Source/Developer/DrawPrimitiveDebugger/Private/SDrawPrimitiveDebugger.cpp b/Engine/Source/Developer/DrawPrimitiveDebugger/Private/SDrawPrimitiveDebugger.cpp
--- a/Engine/Source/Developer/DrawPrimitiveDebugger/Private/SDrawPrimitiveDebugger.cpp
+++ b/Engine/Source/Developer/DrawPrimitiveDebugger/Private/SDrawPrimitiveDebugger.cpp
@@ -183,7 +183,8 @@ void SDrawPrimitiveDebugger::UpdateVisibleRows()
 			{
 				for (const UMaterialInterface* Material : RowData->Materials)
 				{
-					if (IsValid(Material) && Material->GetName().Contains(ActiveFilterString))
+					if (IsValid(Material) && IsValid(Material->GetMaterial()) &&
+						Material->GetMaterial()->GetName().Contains(ActiveFilterString))
 					{
 						bPassesFilter = true;
 						break;
