Modifying the UI for forcing LOD levels with the Primitive Debugger to make the behavior more clear to users without having to view the tooltip. Added missing text localization to the debugger and detail panel.

diff --git a/Engine/Source/Developer/DrawPrimitiveDebugger/Private/SDrawPrimitiveDebugger.cpp b/Engine/Source/Developer/DrawPrimitiveDebugger/Private/SDrawPrimitiveDebugger.cpp
--- a/Engine/Source/Developer/DrawPrimitiveDebugger/Private/SDrawPrimitiveDebugger.cpp
+++ b/Engine/Source/Developer/DrawPrimitiveDebugger/Private/SDrawPrimitiveDebugger.cpp
@@ -27,6 +27,8 @@
 #include "Widgets/Layout/SExpandableArea.h"
 #include "Widgets/Text/STextBlock.h"
 
+#define LOCTEXT_NAMESPACE "PRIMITIVE_DEBUGGER"
+
 #if WITH_PRIMITIVE_DEBUGGER
 
 #define PRIMITIVE_DEBUGGER_SUPPORT_DEBUG_VISUALIZATIONS UE_ENABLE_DEBUG_DRAWING
@@ -61,28 +63,28 @@ void SPrimitiveDebuggerDetailView::Construct(const FArguments& InArgs)
 			SAssignNew(DetailPropertiesWidget, SVerticalBox)
 		]
 	];
-	GenerateDetailPanelEntry("Name:", &SPrimitiveDebuggerDetailView::GetSelectedPrimitiveName, nullptr,
+	GenerateDetailPanelEntry(LOCTEXT("DetailPanel_NameLabel", "Name:"), &SPrimitiveDebuggerDetailView::GetSelectedPrimitiveName, nullptr,
 		&SPrimitiveDebuggerDetailView::GetSelectedPrimitiveName, /* bSupportHighlighting */ true);
-	GenerateDetailPanelEntry("Type:", &SPrimitiveDebuggerDetailView::GetSelectedPrimitiveType, nullptr,
+	GenerateDetailPanelEntry(LOCTEXT("DetailPanel_TypeLabel", "Type:"), &SPrimitiveDebuggerDetailView::GetSelectedPrimitiveType, nullptr,
 		&SPrimitiveDebuggerDetailView::GetSelectedPrimitiveType, /* bSupportHighlighting */ true);
-	GenerateDetailPanelEntry("Actor:", &SPrimitiveDebuggerDetailView::GetSelectedActorName, nullptr,
+	GenerateDetailPanelEntry(LOCTEXT("DetailPanel_ActorLabel", "Actor:"), &SPrimitiveDebuggerDetailView::GetSelectedActorName, nullptr,
 		&SPrimitiveDebuggerDetailView::GetSelectedActorToolTip, /* bSupportHighlighting */ true);
-	GenerateDetailPanelEntry("Actor Class:", &SPrimitiveDebuggerDetailView::GetSelectedActorClassName, nullptr,
+	GenerateDetailPanelEntry(LOCTEXT("DetailPanel_ActorClassLabel", "Actor Class:"), &SPrimitiveDebuggerDetailView::GetSelectedActorClassName, nullptr,
 		&SPrimitiveDebuggerDetailView::GetSelectedActorClassToolTip, /* bSupportHighlighting */ true);
-	GenerateDetailPanelEntry("Location:", &SPrimitiveDebuggerDetailView::GetSelectedLocation);
-	GenerateDetailPanelEntry("Supports Nanite:", &SPrimitiveDebuggerDetailView::GetSelectedPrimitiveSupportsNanite,
+	GenerateDetailPanelEntry(LOCTEXT("DetailPanel_LocationLabel", "Location:"), &SPrimitiveDebuggerDetailView::GetSelectedLocation);
+	GenerateDetailPanelEntry(LOCTEXT("DetailPanel_NaniteSupportLabel", "Supports Nanite:"), &SPrimitiveDebuggerDetailView::GetSelectedPrimitiveSupportsNanite,
 		&SPrimitiveDebuggerDetailView::StaticMeshDataVisibility);
-	GenerateDetailPanelEntry("Nanite Enabled:", &SPrimitiveDebuggerDetailView::GetSelectedPrimitiveNaniteEnabled,
+	GenerateDetailPanelEntry(LOCTEXT("DetailPanel_NaniteEnabledLabel", "Nanite Enabled:"), &SPrimitiveDebuggerDetailView::GetSelectedPrimitiveNaniteEnabled,
 		&SPrimitiveDebuggerDetailView::StaticMeshDataVisibility);
-	GenerateDetailPanelEntry("Current LOD:", &SPrimitiveDebuggerDetailView::GetSelectedLOD,
+	GenerateDetailPanelEntry(LOCTEXT("DetailPanel_CurrentLODLabel", "Current LOD:"), &SPrimitiveDebuggerDetailView::GetSelectedLOD,
 		&SPrimitiveDebuggerDetailView::NonNaniteDataVisibility);
-	GenerateDetailPanelEntry("Available LODs:", &SPrimitiveDebuggerDetailView::GetSelectedNumLODs,
+	GenerateDetailPanelEntry(LOCTEXT("DetailPanel_AvailableLODsLabel", "Available LODs:"), &SPrimitiveDebuggerDetailView::GetSelectedNumLODs,
 		&SPrimitiveDebuggerDetailView::NonNaniteDataVisibility);
-	GenerateDetailPanelEntry("Draw Calls:", &SPrimitiveDebuggerDetailView::GetSelectedDrawCallCount,
+	GenerateDetailPanelEntry(LOCTEXT("DetailPanel_DrawCallsLabel", "Draw Calls:"), &SPrimitiveDebuggerDetailView::GetSelectedDrawCallCount,
 		&SPrimitiveDebuggerDetailView::NonNaniteDataVisibility);
-	GenerateDetailPanelEntry("Triangles:", &SPrimitiveDebuggerDetailView::GetSelectedTriangleCount,
+	GenerateDetailPanelEntry(LOCTEXT("DetailPanel_TrianglesLabel", "Triangles:"), &SPrimitiveDebuggerDetailView::GetSelectedTriangleCount,
 		&SPrimitiveDebuggerDetailView::NonNaniteDataVisibility);
-	GenerateDetailPanelEntry("Bones:", &SPrimitiveDebuggerDetailView::GetSelectedBoneCount,
+	GenerateDetailPanelEntry(LOCTEXT("DetailPanel_BonesLabel", "Bones:"), &SPrimitiveDebuggerDetailView::GetSelectedBoneCount,
 		&SPrimitiveDebuggerDetailView::SkeletalMeshDataVisibility);
 	
 	DetailPropertiesWidget->AddSlot()
@@ -94,7 +96,7 @@ void SPrimitiveDebuggerDetailView::Construct(const FArguments& InArgs)
 		.HeaderContent()
 		[
 			SNew(STextBlock)
-			.Text(FText::FromString("Materials"))
+			.Text(LOCTEXT("DetailPanel_MaterialsLabel", "Materials"))
 		]
 		.BodyContent()
 		[
@@ -112,7 +114,7 @@ void SPrimitiveDebuggerDetailView::Construct(const FArguments& InArgs)
 		.HeaderContent()
 		[
 			SNew(STextBlock)
-			.Text(FText::FromString("Advanced"))
+			.Text(LOCTEXT("DetailPanel_AdvancedOptionsLabel", "Advanced"))
 		]
 		.BodyContent()
 		[
@@ -150,8 +152,9 @@ void SPrimitiveDebuggerDetailView::Tick(const FGeometry& AllottedGeometry, const
 }
 
 static const FText PlaceholderValue = FText::FromString(TEXT("-"));
-static const FText TrueTextValue = FText::FromString(TEXT("true"));
-static const FText FalseTextValue = FText::FromString(TEXT("false"));
+static const FText TrueTextValue = LOCTEXT("DetailPanel_True", "true");
+static const FText FalseTextValue = LOCTEXT("DetailPanel_False", "false");
+static const FText InvalidTextValue = LOCTEXT("DetailPanel_InvalidValue", "INVALID");
 
 void SPrimitiveDebuggerDetailView::UpdateSelection()
 {
@@ -354,13 +357,18 @@ TOptional<int> SPrimitiveDebuggerDetailView::GetSelectedForcedLODValue() const
 	const FPrimitiveRowDataPtr Selection = PrimitiveDebugger.Pin()->GetCurrentSelection();
 	if (Selection.IsValid())
 	{
+		int ForcedLOD = 0;
 		if (SelectedAsStaticMesh.IsValid())
 		{
-			return SelectedAsStaticMesh->ForcedLodModel;
+			ForcedLOD = SelectedAsStaticMesh->ForcedLodModel;
 		}
 		if (SelectedAsSkinnedMesh.IsValid())
 		{
-			return SelectedAsSkinnedMesh->GetForcedLOD();
+			ForcedLOD = SelectedAsSkinnedMesh->GetForcedLOD();
+		}
+		if (ForcedLOD > 0)
+		{
+			return FMath::Clamp(ForcedLOD - 1, 0, GetSelectedNumLODsValue().Get(0) - 1);
 		}
 	}
 	return TOptional<int>();
@@ -381,6 +389,20 @@ TOptional<int> SPrimitiveDebuggerDetailView::GetSelectedNumLODsValue() const
 	return 0;
 }
 
+TOptional<int> SPrimitiveDebuggerDetailView::GetSelectedForcedLODSliderMaxValue() const
+{
+	if (!PrimitiveDebugger.IsValid())
+	{
+		return 0;
+	}
+	const FPrimitiveRowDataPtr Selection = PrimitiveDebugger.Pin()->GetCurrentSelection();
+	if (Selection.IsValid())
+	{
+		return Selection->GetNumLODs() - 1;
+	}
+	return 0;
+}
+
 FText SPrimitiveDebuggerDetailView::GetSelectedTriangleCount() const
 {
 	if (!CurrentLOD)
@@ -401,7 +423,7 @@ FText SPrimitiveDebuggerDetailView::GetSelectedBoneCount() const
 		PlaceholderValue;
 }
 
-void SPrimitiveDebuggerDetailView::GenerateDetailPanelEntry(const FString& Label,
+void SPrimitiveDebuggerDetailView::GenerateDetailPanelEntry(const FText& Label,
 	FText(SPrimitiveDebuggerDetailView::* ValueGetter)() const,
 	EVisibility (SPrimitiveDebuggerDetailView::* VisibilityGetter)() const,
 	FText(SPrimitiveDebuggerDetailView::* TooltipGetter)() const,
@@ -433,7 +455,7 @@ void SPrimitiveDebuggerDetailView::GenerateDetailPanelEntry(const FString& Label
 		.FillWidth(LabelColumnWidth)
 		[
 			SNew(STextBlock)
-			.Text(FText::FromString(Label))
+			.Text(Label)
 			.Justification(ETextJustify::Left)
 		]
 		+SHorizontalBox::Slot()
@@ -551,7 +573,7 @@ void SPrimitiveDebuggerDetailView::CreateMaterialEntry(const UMaterialInterface*
 				.AutoHeight()
 				[
 					SNew(STextBlock)
-					.Text(FText::FromString(TEXT("INVALID")))
+					.Text(InvalidTextValue)
 					.Justification(ETextJustify::Left)
 				];
 			}
@@ -572,7 +594,7 @@ void SPrimitiveDebuggerDetailView::CreateMaterialEntry(const UMaterialInterface*
 			.FillWidth(LabelColumnWidth)
 			[
 				SNew(STextBlock)
-				.Text(FText::FromString(bIsOverlay ? "Overlay" : FString::FromInt(Index)))
+				.Text(bIsOverlay ? LOCTEXT("DetailPanel_OverlayLabel", "Overlay") : FText::FromString(FString::FromInt(Index)))
 				.Justification(ETextJustify::Left)
 			]
 			+SHorizontalBox::Slot()
@@ -600,7 +622,7 @@ void SPrimitiveDebuggerDetailView::CreateMaterialEntry(const UMaterialInterface*
 				.HeaderContent()
 				[
 					SNew(STextBlock)
-					.Text(FText::FromString("Textures"))
+					.Text(LOCTEXT("DetailPanel_TexturesLabel", "Textures"))
 				]
 				.BodyContent()
 				[
@@ -631,6 +653,12 @@ TSharedRef<SVerticalBox> SPrimitiveDebuggerDetailView::GetAdvancedOptionsWidget(
 	const FPrimitiveRowDataPtr Selection = PrimitiveDebugger.Pin()->GetCurrentSelection();
 	if (Selection.IsValid() && Selection->IsPrimitiveValid())
 	{
+		const FText ShowBoundsTooltip = LOCTEXT("DetailPanel_ShowBoundsTooltip", "Should a debug box of this mesh's bounds be displayed? DEVELOPMENT BUILDS ONLY");
+		const FText ShowBonesTooltip = LOCTEXT("DetailPanel_ShowBonesTooltip", "Should a debug display of this mesh's skeleton be displayed? DEVELOPMENT BUILDS ONLY");
+		const FText ForcedLODTooltip = LOCTEXT("DetailPanel_ForcedLODTooltip", "Should a specific LOD level be forced on this primitive?.");
+		const FText ForcedLODIndexTooltip = LOCTEXT("DetailPanel_ForcedLODIndexTooltip", "Controls the forced LOD level of this primitive.");
+		const FText ForceDisableNaniteTooltip = LOCTEXT("DetailPanel_ForcedDisableNaniteTooltip", "Should nanite be force disabled on this component?");
+		
 		AdvancedOptionsWidget->AddSlot()
 		.Padding(Margin)
 		.AutoHeight()
@@ -644,9 +672,9 @@ TSharedRef<SVerticalBox> SPrimitiveDebuggerDetailView::GetAdvancedOptionsWidget(
 			.FillWidth(LabelColumnWidth)
 			[
 				SNew(STextBlock)
-				.Text(FText::FromString("Show Bounds"))
+				.Text(LOCTEXT("DetailPanel_ShowBoundsLabel", "Show Bounds"))
 				.Justification(ETextJustify::Left)
-				.ToolTipText(FText::FromString(TEXT("Should a debug box of this mesh's bounds be displayed? DEVELOPMENT BUILDS ONLY")))
+				.ToolTipText(ShowBoundsTooltip)
 			]
 			+SHorizontalBox::Slot()
 			.HAlign(HAlign_Right)
@@ -655,7 +683,7 @@ TSharedRef<SVerticalBox> SPrimitiveDebuggerDetailView::GetAdvancedOptionsWidget(
 				SNew(SCheckBox)
 				.IsChecked(this, &SPrimitiveDebuggerDetailView::ShowDebugBoundsState)
 				.OnCheckStateChanged(this, &SPrimitiveDebuggerDetailView::OnToggleDebugBounds)
-				.ToolTipText(FText::FromString(TEXT("Should a debug box of this mesh's bounds be displayed? DEVELOPMENT BUILDS ONLY")))
+				.ToolTipText(ShowBoundsTooltip)
 			]
 		];
 		AdvancedOptionsWidget->AddSlot()
@@ -672,9 +700,9 @@ TSharedRef<SVerticalBox> SPrimitiveDebuggerDetailView::GetAdvancedOptionsWidget(
 			.FillWidth(LabelColumnWidth)
 			[
 				SNew(STextBlock)
-				.Text(FText::FromString("Show Bones"))
+				.Text(LOCTEXT("DetailPanel_ShowBonesLabel", "Show Bones"))
 				.Justification(ETextJustify::Left)
-				.ToolTipText(FText::FromString(TEXT("Should a debug display of this mesh's skeleton be displayed? DEVELOPMENT BUILDS ONLY")))
+				.ToolTipText(ShowBonesTooltip)
 			]
 			+SHorizontalBox::Slot()
 			.HAlign(HAlign_Right)
@@ -683,7 +711,32 @@ TSharedRef<SVerticalBox> SPrimitiveDebuggerDetailView::GetAdvancedOptionsWidget(
 				SNew(SCheckBox)
 				.IsChecked(this, &SPrimitiveDebuggerDetailView::ShowDebugBonesState)
 				.OnCheckStateChanged(this, &SPrimitiveDebuggerDetailView::OnToggleDebugBones)
-				.ToolTipText(FText::FromString(TEXT("Should a debug display of this mesh's skeleton be displayed? DEVELOPMENT BUILDS ONLY")))
+				.ToolTipText(ShowBonesTooltip)
+			]
+		];
+		AdvancedOptionsWidget->AddSlot()
+		.Padding(Margin)
+		.AutoHeight()
+		[
+			SNew(SHorizontalBox)
+			.Visibility(this, &SPrimitiveDebuggerDetailView::OptionVisibilityForceLOD)
+			+SHorizontalBox::Slot()
+			.HAlign(HAlign_Left)
+			.FillWidth(LabelColumnWidth)
+			[
+				SNew(STextBlock)
+				.Text(LOCTEXT("DetailPanel_ForceLODLabel", "Force LOD"))
+				.Justification(ETextJustify::Left)
+				.ToolTipText(ForcedLODTooltip)
+			]
+			+SHorizontalBox::Slot()
+			.HAlign(HAlign_Right)
+			.FillWidth(ValueColumnWidth)
+			[
+				SNew(SCheckBox)
+				.IsChecked(this, &SPrimitiveDebuggerDetailView::ForceLODState)
+				.OnCheckStateChanged(this, &SPrimitiveDebuggerDetailView::OnToggleForceLOD)
+				.ToolTipText(ForcedLODTooltip)
 			]
 		];
 		AdvancedOptionsWidget->AddSlot()
@@ -697,9 +750,9 @@ TSharedRef<SVerticalBox> SPrimitiveDebuggerDetailView::GetAdvancedOptionsWidget(
 			.FillWidth(LabelColumnWidth * 2)
 			[
 				SNew(STextBlock)
-				.Text(FText::FromString("Force LOD"))
+				.Text(LOCTEXT("DetailPanel_ForcedLODIndexLabel", "Forced LOD Index"))
 				.Justification(ETextJustify::Left)
-				.ToolTipText(FText::FromString(TEXT("If 0, auto-select LOD level. if >0, force to Forced LOD - 1.")))
+				.ToolTipText(ForcedLODIndexTooltip)
 			]
 			+SHorizontalBox::Slot()
 			.HAlign(HAlign_Fill)
@@ -708,15 +761,17 @@ TSharedRef<SVerticalBox> SPrimitiveDebuggerDetailView::GetAdvancedOptionsWidget(
 				SNew(SNumericEntryBox<int>)
 				.Value(this, &SPrimitiveDebuggerDetailView::GetSelectedForcedLODValue)
 				.MinValue(0)
-				.MaxValue(this, &SPrimitiveDebuggerDetailView::GetSelectedNumLODsValue)
+				.MaxValue(this, &SPrimitiveDebuggerDetailView::GetSelectedForcedLODSliderMaxValue)
 				.MinSliderValue(0)
-				.MaxSliderValue(this, &SPrimitiveDebuggerDetailView::GetSelectedNumLODsValue)
+				.MaxSliderValue(this, &SPrimitiveDebuggerDetailView::GetSelectedForcedLODSliderMaxValue)
 				.Delta(1)
 				.AllowSpin(true)
 				.AllowWheel(true)
 				.WheelStep(1)
+				.UndeterminedString(LOCTEXT("DetailPanel_AutomaticLODPlaceholder", "Auto"))
+				.IsEnabled(this, &SPrimitiveDebuggerDetailView::IsForceLODIndexSliderEnabled)
 				.OnValueChanged(this, &SPrimitiveDebuggerDetailView::HandleForceLOD)
-				.ToolTipText(FText::FromString(TEXT("If 0, auto-select LOD level. if >0, force to Forced LOD - 1.")))
+				.ToolTipText(ForcedLODIndexTooltip)
 				.Justification(ETextJustify::Right)
 			]
 		];
@@ -731,9 +786,9 @@ TSharedRef<SVerticalBox> SPrimitiveDebuggerDetailView::GetAdvancedOptionsWidget(
 			.FillWidth(LabelColumnWidth)
 			[
 				SNew(STextBlock)
-				.Text(FText::FromString("Force Disable Nanite"))
+				.Text(LOCTEXT("DetailPanel_ForceDisableNaniteLabel", "Force Disable Nanite"))
 				.Justification(ETextJustify::Left)
-				.ToolTipText(FText::FromString(TEXT("Should nanite be force disabled on this component?")))
+				.ToolTipText(ForceDisableNaniteTooltip)
 			]
 			+SHorizontalBox::Slot()
 			.HAlign(HAlign_Right)
@@ -742,7 +797,7 @@ TSharedRef<SVerticalBox> SPrimitiveDebuggerDetailView::GetAdvancedOptionsWidget(
 				SNew(SCheckBox)
 				.IsChecked(this, &SPrimitiveDebuggerDetailView::ForceDisableNaniteState)
 				.OnCheckStateChanged(this, &SPrimitiveDebuggerDetailView::OnToggleForceDisableNanite)
-				.ToolTipText(FText::FromString(TEXT("Should nanite be force disabled on this component?")))
+				.ToolTipText(ForceDisableNaniteTooltip)
 			]
 		];
 	}
@@ -760,6 +815,42 @@ EVisibility SPrimitiveDebuggerDetailView::OptionVisibilityForceDisableNanite() c
 	return bSelectionSupportsNanite && SelectedAsStaticMesh.IsValid() ? EVisibility::Visible : EVisibility::Collapsed;
 }
 
+ECheckBoxState SPrimitiveDebuggerDetailView::ForceLODState() const
+{
+	if (!PrimitiveDebugger.IsValid())
+	{
+		return ECheckBoxState::Undetermined;
+	}
+	const TSharedPtr<SDrawPrimitiveDebugger> DebuggerInstance = PrimitiveDebugger.Pin();
+	const FPrimitiveComponentId Selection = DebuggerInstance->GetCurrentSelectionId();
+	return DebuggerInstance->DoesEntryHaveForcedLOD(Selection) ? ECheckBoxState::Checked : ECheckBoxState::Unchecked;
+}
+
+bool SPrimitiveDebuggerDetailView::IsForceLODIndexSliderEnabled() const
+{
+	const TSharedPtr<SDrawPrimitiveDebugger> DebuggerInstance = PrimitiveDebugger.Pin();
+	const FPrimitiveComponentId Selection = DebuggerInstance->GetCurrentSelectionId();
+	return DebuggerInstance->DoesEntryHaveForcedLOD(Selection);
+}
+
+void SPrimitiveDebuggerDetailView::OnToggleForceLOD(ECheckBoxState state)
+{
+	if (!PrimitiveDebugger.IsValid())
+	{
+		return;
+	}
+	const TSharedPtr<SDrawPrimitiveDebugger> DebuggerInstance = PrimitiveDebugger.Pin();
+	const FPrimitiveComponentId Selection = DebuggerInstance->GetCurrentSelectionId();
+	if (state == ECheckBoxState::Unchecked)
+	{
+		DebuggerInstance->SetForcedLODForEntry(Selection, 0);
+	}
+	else if (state == ECheckBoxState::Checked)
+	{
+		DebuggerInstance->SetForcedLODForEntry(Selection, CurrentLOD->LODIndex + 1);
+	}
+}
+
 ECheckBoxState SPrimitiveDebuggerDetailView::ForceDisableNaniteState() const
 {
 	if (SelectedAsStaticMesh.IsValid())
@@ -800,7 +891,7 @@ void SPrimitiveDebuggerDetailView::HandleForceLOD(int ForcedLOD)
 	}
 	const TSharedPtr<SDrawPrimitiveDebugger> DebuggerInstance = PrimitiveDebugger.Pin();
 	const FPrimitiveComponentId Selection = DebuggerInstance->GetCurrentSelectionId();
-	DebuggerInstance->SetForcedLODForEntry(Selection, ForcedLOD);
+	DebuggerInstance->SetForcedLODForEntry(Selection, ForcedLOD + 1);
 }
 	
 ECheckBoxState SPrimitiveDebuggerDetailView::ShowDebugBoundsState() const
@@ -898,10 +989,10 @@ void SDrawPrimitiveDebugger::Construct(const FArguments& InArgs)
 	const FName PinColumn("Pin");
 	const FName NameColumn("Name");
 	const FName ActorColumn("Actor");
-	AddColumn(FText::FromString("Visible"), VisibilityColumn);
-	AddColumn(FText::FromString("Pinned"), PinColumn);
-	AddColumn(FText::FromString("Name"), NameColumn);
-	AddColumn(FText::FromString("Actor"), ActorColumn);
+	AddColumn(LOCTEXT("VisbleColumnLabel", "Visible"), VisibilityColumn);
+	AddColumn(LOCTEXT("PinnedColumnLabel", "Pinned"), PinColumn);
+	AddColumn(LOCTEXT("NameColumnLabel", "Name"), NameColumn);
+	AddColumn(LOCTEXT("ActorColumnLabel", "Actor"), ActorColumn);
 
 	FilterText = FText::GetEmpty();
 	IDrawPrimitiveDebugger::Get().CaptureSingleFrame();
@@ -938,7 +1029,7 @@ void SDrawPrimitiveDebugger::Construct(const FArguments& InArgs)
 			.AutoWidth()
 			[
 				SNew(SButton)
-				.Text(FText::FromString("Refresh"))
+				.Text(LOCTEXT("RefreshButtonLabel", "Refresh"))
 				.IsEnabled(this, &SDrawPrimitiveDebugger::CanCaptureSingleFrame)
 				.OnClicked(this, &SDrawPrimitiveDebugger::OnRefreshClick)
 			]
@@ -947,7 +1038,7 @@ void SDrawPrimitiveDebugger::Construct(const FArguments& InArgs)
 			.AutoWidth()
 			[
 				SNew(SButton)
-				.Text(FText::FromString("Save to CSV"))
+				.Text(LOCTEXT("SaveToCSVButtonLabel", "Save to CSV"))
 				.OnClicked(this, &SDrawPrimitiveDebugger::OnSaveClick)
 			]
 			/*+ SHorizontalBox::Slot()
@@ -1378,6 +1469,44 @@ void SDrawPrimitiveDebugger::SetForcedLODForEntry(FPrimitiveComponentId EntryId,
 	}
 }
 
+void SDrawPrimitiveDebugger::ResetForcedLODForEntry(FPrimitiveComponentId EntryId)
+{
+	if (FPrimitiveDebuggerEntry* Entry = Entries.Find(EntryId))
+	{
+		if (!Entry->bHasForcedLOD || !Entry->Data.IsValid() || !Entry->Data->IsPrimitiveValid())
+		{
+			return;
+		}
+		if (UStaticMeshComponent* StaticMesh = Cast<UStaticMeshComponent>(Entry->Data->ComponentUObject.Get()))
+		{
+			StaticMesh->SetForcedLodModel(Entry->DesiredForcedLOD);
+		}
+		else if (USkinnedMeshComponent* SkinnedMesh = Cast<USkinnedMeshComponent>(Entry->Data->ComponentUObject.Get()))
+		{
+			SkinnedMesh->SetForcedLOD(Entry->DesiredForcedLOD);
+		}
+		Entry->bHasForcedLOD = false;
+	}
+}
+
+bool SDrawPrimitiveDebugger::DoesEntryHaveForcedLOD(FPrimitiveComponentId EntryId) const
+{
+	if (const FPrimitiveDebuggerEntry* Entry = Entries.Find(EntryId))
+	{
+		bool bHasForcedLOD = false;
+		if (const UStaticMeshComponent* StaticMesh = Cast<UStaticMeshComponent>(Entry->Data->ComponentUObject.Get()))
+		{
+			bHasForcedLOD = StaticMesh->ForcedLodModel != 0;
+		}
+		else if (const USkinnedMeshComponent* SkinnedMesh = Cast<USkinnedMeshComponent>(Entry->Data->ComponentUObject.Get()))
+		{
+			bHasForcedLOD = SkinnedMesh->GetForcedLOD() != 0;
+		}
+		return bHasForcedLOD;
+	}
+	return false;
+}
+
 void SDrawPrimitiveDebugger::SetForceDisabledNaniteForEntry(FPrimitiveComponentId EntryId, bool bForceDisableNanite)
 {
 	if (FPrimitiveDebuggerEntry* Entry = Entries.Find(EntryId))
@@ -1590,7 +1719,7 @@ void SDrawPrimitiveDebugger::ResetDebuggerChanges()
 		}
 		if (Entry.bHasForcedLOD)
 		{
-			SetForcedLODForEntry(PrimitiveId, Entry.DesiredForcedLOD);
+			ResetForcedLODForEntry(PrimitiveId);
 		}
 		if (Entry.bHasForceDisabledNanite)
 		{
@@ -1794,14 +1923,14 @@ TSharedRef<SWidget> SDrawPrimitiveDebuggerListViewRow::MakeCellWidget(const int3
 			SCOPE_CYCLE_COUNTER(STAT_PrimitiveDebuggerMakeCellActorClass);
 			Value = RowDataPtr->Owner.IsValid() && IsValid(RowDataPtr->Owner->GetClass()) ?
 				FText::FromString(RowDataPtr->Owner->GetClass()->GetName()) :
-				FText::FromString("INVALID");
+				InvalidTextValue;
 		}
 		else if (InColumnId.IsEqual(ActorColumn))
 		{
 			SCOPE_CYCLE_COUNTER(STAT_PrimitiveDebuggerMakeCellActor);
 			Value = RowDataPtr->Owner.IsValid() ?
 				FText::FromString(RowDataPtr->GetOwnerName()) :
-				FText::FromString("INVALID");
+				InvalidTextValue;
 		}
 		else
 		{
@@ -1839,4 +1968,6 @@ ECheckBoxState SDrawPrimitiveDebuggerListViewRow::IsPinned() const
 
 END_SLATE_FUNCTION_BUILD_OPTIMIZATION
 
+#undef LOCTEXT_NAMESPACE
+
 #endif
\ No newline at end of file
diff --git a/Engine/Source/Developer/DrawPrimitiveDebugger/Private/SDrawPrimitiveDebugger.h b/Engine/Source/Developer/DrawPrimitiveDebugger/Private/SDrawPrimitiveDebugger.h
--- a/Engine/Source/Developer/DrawPrimitiveDebugger/Private/SDrawPrimitiveDebugger.h
+++ b/Engine/Source/Developer/DrawPrimitiveDebugger/Private/SDrawPrimitiveDebugger.h
@@ -79,10 +79,11 @@ private:
 	TOptional<int> GetSelectedLODValue() const;
 	TOptional<int> GetSelectedForcedLODValue() const;
 	TOptional<int> GetSelectedNumLODsValue() const;
+	TOptional<int> GetSelectedForcedLODSliderMaxValue() const;
 	FText GetSelectedTriangleCount() const;
 	FText GetSelectedBoneCount() const;
 
-	void GenerateDetailPanelEntry(const FString& Label,
+	void GenerateDetailPanelEntry(const FText& Label,
 		FText (SPrimitiveDebuggerDetailView::* ValueGetter)() const,
 		EVisibility (SPrimitiveDebuggerDetailView::* VisibilityGetter)() const = nullptr,
 		FText (SPrimitiveDebuggerDetailView::* TooltipGetter)() const = nullptr,
@@ -94,6 +95,9 @@ private:
 	EVisibility OptionVisibilityForceLOD() const;
 	EVisibility OptionVisibilityForceDisableNanite() const;
 
+	ECheckBoxState ForceLODState() const;
+	bool IsForceLODIndexSliderEnabled() const;
+	void OnToggleForceLOD(ECheckBoxState state);
 	void HandleForceLOD(int ForcedLOD);
 	
 	ECheckBoxState ShowDebugBoundsState() const;
@@ -151,6 +155,8 @@ public:
 	bool IsEntryPinned(FPrimitiveRowDataPtr Data) const;
 
 	void SetForcedLODForEntry(FPrimitiveComponentId EntryId, int32 NewForcedLOD);
+	void ResetForcedLODForEntry(FPrimitiveComponentId EntryId);
+	bool DoesEntryHaveForcedLOD(FPrimitiveComponentId EntryId) const;
 	void SetForceDisabledNaniteForEntry(FPrimitiveComponentId EntryId, bool bForceDisableNanite);
 	
 	void SetShowDebugBoundsForEntry(FPrimitiveComponentId EntryId, bool bShowDebugBounds);
